Add basic UnrealCodeAnalyzer tests.

diff --git a/Engine/Source/Developer/UnrealCodeAnalyzerTests/Private/Inheritance.cpp b/Engine/Source/Developer/UnrealCodeAnalyzerTests/Private/Inheritance.cpp
new file mode 100644
--- /dev/null
+++ b/Engine/Source/Developer/UnrealCodeAnalyzerTests/Private/Inheritance.cpp
@@ -0,0 +1,10 @@
+#include "UnrealCodeAnalyzerTestsPCH.h"
+#include "Inheritance.h"
+
+class FTestDerivedClass : public FTestBaseClass
+{ };
+
+static void Function_Inheritance()
+{
+	FTestDerivedClass TestDerivedClass;
+}
diff --git a/Engine/Source/Developer/UnrealCodeAnalyzerTests/Private/Inheritance.h b/Engine/Source/Developer/UnrealCodeAnalyzerTests/Private/Inheritance.h
new file mode 100644
--- /dev/null
+++ b/Engine/Source/Developer/UnrealCodeAnalyzerTests/Private/Inheritance.h
@@ -0,0 +1,4 @@
+#pragma once
+
+class FTestBaseClass
+{ };
diff --git a/Engine/Source/Developer/UnrealCodeAnalyzerTests/Private/TemplatedClassInstance.cpp b/Engine/Source/Developer/UnrealCodeAnalyzerTests/Private/TemplatedClassInstance.cpp
new file mode 100644
--- /dev/null
+++ b/Engine/Source/Developer/UnrealCodeAnalyzerTests/Private/TemplatedClassInstance.cpp
@@ -0,0 +1,9 @@
+// Copyright 1998-2015 Epic Games, Inc. All Rights Reserved.
+
+#include "UnrealCodeAnalyzerTestsPCH.h"
+#include "TemplatedClassInstance.h"
+
+static void Function_TemplatedClassInstance()
+{
+	FTemplateClass<FTestClass> TemplateClass;
+}
diff --git a/Engine/Source/Developer/UnrealCodeAnalyzerTests/Private/TemplatedClassInstance.h b/Engine/Source/Developer/UnrealCodeAnalyzerTests/Private/TemplatedClassInstance.h
new file mode 100644
--- /dev/null
+++ b/Engine/Source/Developer/UnrealCodeAnalyzerTests/Private/TemplatedClassInstance.h
@@ -0,0 +1,8 @@
+// Copyright 1998-2015 Epic Games, Inc. All Rights Reserved.
+
+#pragma once
+#include "UseClass.h"
+
+template <typename T>
+class FTemplateClass
+{ };
diff --git a/Engine/Source/Developer/UnrealCodeAnalyzerTests/Private/TemplatedFunctionInstance_Explicit.cpp b/Engine/Source/Developer/UnrealCodeAnalyzerTests/Private/TemplatedFunctionInstance_Explicit.cpp
new file mode 100644
--- /dev/null
+++ b/Engine/Source/Developer/UnrealCodeAnalyzerTests/Private/TemplatedFunctionInstance_Explicit.cpp
@@ -0,0 +1,9 @@
+// Copyright 1998-2015 Epic Games, Inc. All Rights Reserved.
+
+#include "UnrealCodeAnalyzerTestsPCH.h"
+#include "TemplatedFunctionInstance_Explicit.h"
+
+static void Function_TemplatedFunctionInstance_Explicit()
+{
+	TemplateFunction_Explicit<FTestClass>();
+}
diff --git a/Engine/Source/Developer/UnrealCodeAnalyzerTests/Private/TemplatedFunctionInstance_Explicit.h b/Engine/Source/Developer/UnrealCodeAnalyzerTests/Private/TemplatedFunctionInstance_Explicit.h
new file mode 100644
--- /dev/null
+++ b/Engine/Source/Developer/UnrealCodeAnalyzerTests/Private/TemplatedFunctionInstance_Explicit.h
@@ -0,0 +1,8 @@
+// Copyright 1998-2015 Epic Games, Inc. All Rights Reserved.
+
+#pragma once
+#include "UseClass.h"
+
+template <typename T>
+void TemplateFunction_Explicit()
+{ }
diff --git a/Engine/Source/Developer/UnrealCodeAnalyzerTests/Private/TemplatedFunctionInstance_Implicit.cpp b/Engine/Source/Developer/UnrealCodeAnalyzerTests/Private/TemplatedFunctionInstance_Implicit.cpp
new file mode 100644
--- /dev/null
+++ b/Engine/Source/Developer/UnrealCodeAnalyzerTests/Private/TemplatedFunctionInstance_Implicit.cpp
@@ -0,0 +1,9 @@
+// Copyright 1998-2015 Epic Games, Inc. All Rights Reserved.
+
+#include "UnrealCodeAnalyzerTestsPCH.h"
+#include "TemplatedFunctionInstance_Implicit.h"
+
+static void Function_TemplatedFunctionInstance_Implicit()
+{
+	TemplateFunction_Implicit((FTestClass*)nullptr);
+}
diff --git a/Engine/Source/Developer/UnrealCodeAnalyzerTests/Private/TemplatedFunctionInstance_Implicit.h b/Engine/Source/Developer/UnrealCodeAnalyzerTests/Private/TemplatedFunctionInstance_Implicit.h
new file mode 100644
--- /dev/null
+++ b/Engine/Source/Developer/UnrealCodeAnalyzerTests/Private/TemplatedFunctionInstance_Implicit.h
@@ -0,0 +1,8 @@
+// Copyright 1998-2015 Epic Games, Inc. All Rights Reserved.
+
+#pragma once
+#include "UseClass.h"
+
+template <typename T>
+void TemplateFunction_Implicit(T* Param)
+{ }
diff --git a/Engine/Source/Developer/UnrealCodeAnalyzerTests/Private/UnrealCodeAnalyzerTests.cpp b/Engine/Source/Developer/UnrealCodeAnalyzerTests/Private/UnrealCodeAnalyzerTests.cpp
new file mode 100644
--- /dev/null
+++ b/Engine/Source/Developer/UnrealCodeAnalyzerTests/Private/UnrealCodeAnalyzerTests.cpp
@@ -0,0 +1,6 @@
+// Copyright 1998-2015 Epic Games, Inc. All Rights Reserved.
+
+#include "UnrealCodeAnalyzerTestsPCH.h"
+
+
+IMPLEMENT_MODULE(FDefaultModuleImpl, UnrealCodeAnalyzerTests);
diff --git a/Engine/Source/Developer/UnrealCodeAnalyzerTests/Private/UnrealCodeAnalyzerTestsPCH.h b/Engine/Source/Developer/UnrealCodeAnalyzerTests/Private/UnrealCodeAnalyzerTestsPCH.h
new file mode 100644
--- /dev/null
+++ b/Engine/Source/Developer/UnrealCodeAnalyzerTests/Private/UnrealCodeAnalyzerTestsPCH.h
@@ -0,0 +1,4 @@
+// Copyright 1998-2015 Epic Games, Inc. All Rights Reserved.
+
+#include "UnrealCodeAnalyzerTests.h"
+#include "ModuleManager.h"
diff --git a/Engine/Source/Developer/UnrealCodeAnalyzerTests/Private/UseClass.cpp b/Engine/Source/Developer/UnrealCodeAnalyzerTests/Private/UseClass.cpp
new file mode 100644
--- /dev/null
+++ b/Engine/Source/Developer/UnrealCodeAnalyzerTests/Private/UseClass.cpp
@@ -0,0 +1,9 @@
+// Copyright 1998-2015 Epic Games, Inc. All Rights Reserved.
+
+#include "UnrealCodeAnalyzerTestsPCH.h"
+#include "UseClass.h"
+
+static void Function_UseClass()
+{
+	FTestClass TestClass;
+}
diff --git a/Engine/Source/Developer/UnrealCodeAnalyzerTests/Private/UseClass.h b/Engine/Source/Developer/UnrealCodeAnalyzerTests/Private/UseClass.h
new file mode 100644
--- /dev/null
+++ b/Engine/Source/Developer/UnrealCodeAnalyzerTests/Private/UseClass.h
@@ -0,0 +1,6 @@
+// Copyright 1998-2015 Epic Games, Inc. All Rights Reserved.
+
+#pragma once
+
+class FTestClass
+{ };
diff --git a/Engine/Source/Developer/UnrealCodeAnalyzerTests/Private/UseClassStaticField.cpp b/Engine/Source/Developer/UnrealCodeAnalyzerTests/Private/UseClassStaticField.cpp
new file mode 100644
--- /dev/null
+++ b/Engine/Source/Developer/UnrealCodeAnalyzerTests/Private/UseClassStaticField.cpp
@@ -0,0 +1,9 @@
+// Copyright 1998-2015 Epic Games, Inc. All Rights Reserved.
+
+#include "UnrealCodeAnalyzerTestsPCH.h"
+#include "UseClassStaticField.h"
+
+static void Function_UseClassStaticField()
+{
+	FTestClassWithStaticField::StaticField = 10;
+}
diff --git a/Engine/Source/Developer/UnrealCodeAnalyzerTests/Private/UseClassStaticField.h b/Engine/Source/Developer/UnrealCodeAnalyzerTests/Private/UseClassStaticField.h
new file mode 100644
--- /dev/null
+++ b/Engine/Source/Developer/UnrealCodeAnalyzerTests/Private/UseClassStaticField.h
@@ -0,0 +1,11 @@
+// Copyright 1998-2015 Epic Games, Inc. All Rights Reserved.
+
+#pragma once
+
+class FTestClassWithStaticField
+{
+public:
+	static int StaticField;
+};
+
+int FTestClassWithStaticField::StaticField = 5;
diff --git a/Engine/Source/Developer/UnrealCodeAnalyzerTests/Private/UseClassStaticMember.cpp b/Engine/Source/Developer/UnrealCodeAnalyzerTests/Private/UseClassStaticMember.cpp
new file mode 100644
--- /dev/null
+++ b/Engine/Source/Developer/UnrealCodeAnalyzerTests/Private/UseClassStaticMember.cpp
@@ -0,0 +1,9 @@
+// Copyright 1998-2015 Epic Games, Inc. All Rights Reserved.
+
+#include "UnrealCodeAnalyzerTestsPCH.h"
+#include "UseClassStaticMember.h"
+
+static void Function_UseClassStaticMember()
+{
+	FTestClassWithStaticMember::StaticMember();
+}
diff --git a/Engine/Source/Developer/UnrealCodeAnalyzerTests/Private/UseClassStaticMember.h b/Engine/Source/Developer/UnrealCodeAnalyzerTests/Private/UseClassStaticMember.h
new file mode 100644
--- /dev/null
+++ b/Engine/Source/Developer/UnrealCodeAnalyzerTests/Private/UseClassStaticMember.h
@@ -0,0 +1,10 @@
+// Copyright 1998-2015 Epic Games, Inc. All Rights Reserved.
+
+#pragma once
+
+class FTestClassWithStaticMember
+{
+public:
+	static void StaticMember()
+	{ }
+};
diff --git a/Engine/Source/Developer/UnrealCodeAnalyzerTests/Private/UseClassWithNonStaticField.cpp b/Engine/Source/Developer/UnrealCodeAnalyzerTests/Private/UseClassWithNonStaticField.cpp
new file mode 100644
--- /dev/null
+++ b/Engine/Source/Developer/UnrealCodeAnalyzerTests/Private/UseClassWithNonStaticField.cpp
@@ -0,0 +1,9 @@
+// Copyright 1998-2015 Epic Games, Inc. All Rights Reserved.
+
+#include "UnrealCodeAnalyzerTestsPCH.h"
+#include "UseClassWithNonStaticField.h"
+
+static void Function_UseClassWithNonStaticField()
+{
+	FTestClassWithNonStaticField TestClassWithNonStaticField;
+}
diff --git a/Engine/Source/Developer/UnrealCodeAnalyzerTests/Private/UseClassWithNonStaticField.h b/Engine/Source/Developer/UnrealCodeAnalyzerTests/Private/UseClassWithNonStaticField.h
new file mode 100644
--- /dev/null
+++ b/Engine/Source/Developer/UnrealCodeAnalyzerTests/Private/UseClassWithNonStaticField.h
@@ -0,0 +1,10 @@
+// Copyright 1998-2015 Epic Games, Inc. All Rights Reserved.
+
+#pragma once
+#include "UseClass.h"
+
+class FTestClassWithNonStaticField
+{
+public:
+	FTestClass NonStaticField;
+};
diff --git a/Engine/Source/Developer/UnrealCodeAnalyzerTests/Private/UseClassWithNonStaticPtrField.cpp b/Engine/Source/Developer/UnrealCodeAnalyzerTests/Private/UseClassWithNonStaticPtrField.cpp
new file mode 100644
--- /dev/null
+++ b/Engine/Source/Developer/UnrealCodeAnalyzerTests/Private/UseClassWithNonStaticPtrField.cpp
@@ -0,0 +1,9 @@
+// Copyright 1998-2015 Epic Games, Inc. All Rights Reserved.
+
+#include "UnrealCodeAnalyzerTestsPCH.h"
+#include "UseClassWithNonStaticPtrField.h"
+
+static void Function_UseClassWithNonStaticPtrField()
+{
+	FTestClassWithNonStaticPtrField TestClassWithNonStaticPtrField;
+}
diff --git a/Engine/Source/Developer/UnrealCodeAnalyzerTests/Private/UseClassWithNonStaticPtrField.h b/Engine/Source/Developer/UnrealCodeAnalyzerTests/Private/UseClassWithNonStaticPtrField.h
new file mode 100644
--- /dev/null
+++ b/Engine/Source/Developer/UnrealCodeAnalyzerTests/Private/UseClassWithNonStaticPtrField.h
@@ -0,0 +1,10 @@
+// Copyright 1998-2015 Epic Games, Inc. All Rights Reserved.
+
+#pragma once
+#include "UseClass.h"
+
+class FTestClassWithNonStaticPtrField
+{
+public:
+	FTestClass* NonStaticField;
+};
diff --git a/Engine/Source/Developer/UnrealCodeAnalyzerTests/Private/UseElifMacro.cpp b/Engine/Source/Developer/UnrealCodeAnalyzerTests/Private/UseElifMacro.cpp
new file mode 100644
--- /dev/null
+++ b/Engine/Source/Developer/UnrealCodeAnalyzerTests/Private/UseElifMacro.cpp
@@ -0,0 +1,15 @@
+// Copyright 1998-2015 Epic Games, Inc. All Rights Reserved.
+
+#include "UnrealCodeAnalyzerTestsPCH.h"
+#include "UseElifMacro.h"
+
+static int Function_UseElifMacro()
+{
+#if IF_MACRO
+	return 0;
+#elif ELIF_MACRO
+	return 1;
+#else
+	return 2;
+#endif
+}
diff --git a/Engine/Source/Developer/UnrealCodeAnalyzerTests/Private/UseElifMacro.h b/Engine/Source/Developer/UnrealCodeAnalyzerTests/Private/UseElifMacro.h
new file mode 100644
--- /dev/null
+++ b/Engine/Source/Developer/UnrealCodeAnalyzerTests/Private/UseElifMacro.h
@@ -0,0 +1,6 @@
+// Copyright 1998-2015 Epic Games, Inc. All Rights Reserved.
+
+#pragma once
+
+#define IF_MACRO 0
+#define ELIF_MACRO 1
\ No newline at end of file
diff --git a/Engine/Source/Developer/UnrealCodeAnalyzerTests/Private/UseFunction.cpp b/Engine/Source/Developer/UnrealCodeAnalyzerTests/Private/UseFunction.cpp
new file mode 100644
--- /dev/null
+++ b/Engine/Source/Developer/UnrealCodeAnalyzerTests/Private/UseFunction.cpp
@@ -0,0 +1,9 @@
+// Copyright 1998-2015 Epic Games, Inc. All Rights Reserved.
+
+#include "UnrealCodeAnalyzerTestsPCH.h"
+#include "UseFunction.h"
+
+static void Function_UseFunction()
+{
+	TestFunction();
+}
\ No newline at end of file
diff --git a/Engine/Source/Developer/UnrealCodeAnalyzerTests/Private/UseFunction.h b/Engine/Source/Developer/UnrealCodeAnalyzerTests/Private/UseFunction.h
new file mode 100644
--- /dev/null
+++ b/Engine/Source/Developer/UnrealCodeAnalyzerTests/Private/UseFunction.h
@@ -0,0 +1,5 @@
+// Copyright 1998-2015 Epic Games, Inc. All Rights Reserved.
+
+#pragma once
+
+void TestFunction();
diff --git a/Engine/Source/Developer/UnrealCodeAnalyzerTests/Private/UseIfMacro.cpp b/Engine/Source/Developer/UnrealCodeAnalyzerTests/Private/UseIfMacro.cpp
new file mode 100644
--- /dev/null
+++ b/Engine/Source/Developer/UnrealCodeAnalyzerTests/Private/UseIfMacro.cpp
@@ -0,0 +1,13 @@
+// Copyright 1998-2015 Epic Games, Inc. All Rights Reserved.
+
+#include "UnrealCodeAnalyzerTestsPCH.h"
+#include "UseIfMacro.h"
+
+static int Function_UseIfMacro()
+{
+#if IF_MACRO
+	return 0;
+#else
+	return 1;
+#endif
+}
diff --git a/Engine/Source/Developer/UnrealCodeAnalyzerTests/Private/UseIfMacro.h b/Engine/Source/Developer/UnrealCodeAnalyzerTests/Private/UseIfMacro.h
new file mode 100644
--- /dev/null
+++ b/Engine/Source/Developer/UnrealCodeAnalyzerTests/Private/UseIfMacro.h
@@ -0,0 +1,5 @@
+// Copyright 1998-2015 Epic Games, Inc. All Rights Reserved.
+
+#pragma once
+
+#define IF_MACRO 1
\ No newline at end of file
diff --git a/Engine/Source/Developer/UnrealCodeAnalyzerTests/Private/UseVariable.cpp b/Engine/Source/Developer/UnrealCodeAnalyzerTests/Private/UseVariable.cpp
new file mode 100644
--- /dev/null
+++ b/Engine/Source/Developer/UnrealCodeAnalyzerTests/Private/UseVariable.cpp
@@ -0,0 +1,9 @@
+// Copyright 1998-2015 Epic Games, Inc. All Rights Reserved.
+
+#include "UnrealCodeAnalyzerTestsPCH.h"
+#include "UseVariable.h"
+
+static void Function_UseVariable()
+{
+	TestVariable = 0;
+}
diff --git a/Engine/Source/Developer/UnrealCodeAnalyzerTests/Private/UseVariable.h b/Engine/Source/Developer/UnrealCodeAnalyzerTests/Private/UseVariable.h
new file mode 100644
--- /dev/null
+++ b/Engine/Source/Developer/UnrealCodeAnalyzerTests/Private/UseVariable.h
@@ -0,0 +1,5 @@
+// Copyright 1998-2015 Epic Games, Inc. All Rights Reserved.
+
+#pragma once
+
+int TestVariable;
diff --git a/Engine/Source/Developer/UnrealCodeAnalyzerTests/Public/UnrealCodeAnalyzerTests.h b/Engine/Source/Developer/UnrealCodeAnalyzerTests/Public/UnrealCodeAnalyzerTests.h
new file mode 100644
--- /dev/null
+++ b/Engine/Source/Developer/UnrealCodeAnalyzerTests/Public/UnrealCodeAnalyzerTests.h
@@ -0,0 +1,5 @@
+// Copyright 1998-2015 Epic Games, Inc. All Rights Reserved.
+
+#pragma once
+
+#include "Core.h"
diff --git a/Engine/Source/Developer/UnrealCodeAnalyzerTests/UnrealCodeAnalyzerTests.Build.cs b/Engine/Source/Developer/UnrealCodeAnalyzerTests/UnrealCodeAnalyzerTests.Build.cs
new file mode 100644
--- /dev/null
+++ b/Engine/Source/Developer/UnrealCodeAnalyzerTests/UnrealCodeAnalyzerTests.Build.cs
@@ -0,0 +1,11 @@
+// Copyright 1998-2015 Epic Games, Inc. All Rights Reserved.
+
+using UnrealBuildTool;
+
+public class UnrealCodeAnalyzerTests : ModuleRules
+{
+	public UnrealCodeAnalyzerTests(TargetInfo Target)
+	{
+		PublicDependencyModuleNames.Add("Core");
+	}
+}
