CIS fix.

diff --git a/Engine/Source/Developer/UnrealCodeAnalyzerTests/Private/Inheritance.cpp b/Engine/Source/Developer/UnrealCodeAnalyzerTests/Private/Inheritance.cpp
--- a/Engine/Source/Developer/UnrealCodeAnalyzerTests/Private/Inheritance.cpp
+++ b/Engine/Source/Developer/UnrealCodeAnalyzerTests/Private/Inheritance.cpp
@@ -6,5 +6,7 @@ class FTestDerivedClass : public FTestBaseClass
 
 static void Function_Inheritance()
 {
+#ifdef UNREAL_CODE_ANALYZER
 	FTestDerivedClass TestDerivedClass;
+#endif
 }
diff --git a/Engine/Source/Developer/UnrealCodeAnalyzerTests/Private/TemplatedClassInstance.cpp b/Engine/Source/Developer/UnrealCodeAnalyzerTests/Private/TemplatedClassInstance.cpp
--- a/Engine/Source/Developer/UnrealCodeAnalyzerTests/Private/TemplatedClassInstance.cpp
+++ b/Engine/Source/Developer/UnrealCodeAnalyzerTests/Private/TemplatedClassInstance.cpp
@@ -5,5 +5,7 @@
 
 static void Function_TemplatedClassInstance()
 {
+#ifdef UNREAL_CODE_ANALYZER
 	FTemplateClass<FTestClass> TemplateClass;
+#endif
 }
diff --git a/Engine/Source/Developer/UnrealCodeAnalyzerTests/Private/UseClass.cpp b/Engine/Source/Developer/UnrealCodeAnalyzerTests/Private/UseClass.cpp
--- a/Engine/Source/Developer/UnrealCodeAnalyzerTests/Private/UseClass.cpp
+++ b/Engine/Source/Developer/UnrealCodeAnalyzerTests/Private/UseClass.cpp
@@ -5,5 +5,7 @@
 
 static void Function_UseClass()
 {
+#ifdef UNREAL_CODE_ANALYZER
 	FTestClass TestClass;
+#endif
 }
diff --git a/Engine/Source/Developer/UnrealCodeAnalyzerTests/Private/UseClassWithNonStaticField.cpp b/Engine/Source/Developer/UnrealCodeAnalyzerTests/Private/UseClassWithNonStaticField.cpp
--- a/Engine/Source/Developer/UnrealCodeAnalyzerTests/Private/UseClassWithNonStaticField.cpp
+++ b/Engine/Source/Developer/UnrealCodeAnalyzerTests/Private/UseClassWithNonStaticField.cpp
@@ -5,5 +5,7 @@
 
 static void Function_UseClassWithNonStaticField()
 {
+#ifdef UNREAL_CODE_ANALYZER
 	FTestClassWithNonStaticField TestClassWithNonStaticField;
+#endif
 }
diff --git a/Engine/Source/Developer/UnrealCodeAnalyzerTests/Private/UseClassWithNonStaticPtrField.cpp b/Engine/Source/Developer/UnrealCodeAnalyzerTests/Private/UseClassWithNonStaticPtrField.cpp
--- a/Engine/Source/Developer/UnrealCodeAnalyzerTests/Private/UseClassWithNonStaticPtrField.cpp
+++ b/Engine/Source/Developer/UnrealCodeAnalyzerTests/Private/UseClassWithNonStaticPtrField.cpp
@@ -5,5 +5,7 @@
 
 static void Function_UseClassWithNonStaticPtrField()
 {
+#ifdef UNREAL_CODE_ANALYZER
 	FTestClassWithNonStaticPtrField TestClassWithNonStaticPtrField;
+#endif
 }
