CIS fix.

diff --git a/Engine/Source/Developer/UnrealCodeAnalyzerTests/Private/UseElifMacro.cpp b/Engine/Source/Developer/UnrealCodeAnalyzerTests/Private/UseElifMacro.cpp
--- a/Engine/Source/Developer/UnrealCodeAnalyzerTests/Private/UseElifMacro.cpp
+++ b/Engine/Source/Developer/UnrealCodeAnalyzerTests/Private/UseElifMacro.cpp
@@ -5,7 +5,7 @@
 
 static int Function_UseElifMacro()
 {
-#if IF_MACRO
+#if TEST_IF_MACRO
 	return 0;
 #elif ELIF_MACRO
 	return 1;
diff --git a/Engine/Source/Developer/UnrealCodeAnalyzerTests/Private/UseElifMacro.h b/Engine/Source/Developer/UnrealCodeAnalyzerTests/Private/UseElifMacro.h
--- a/Engine/Source/Developer/UnrealCodeAnalyzerTests/Private/UseElifMacro.h
+++ b/Engine/Source/Developer/UnrealCodeAnalyzerTests/Private/UseElifMacro.h
@@ -2,5 +2,5 @@
 
 #pragma once
 
-#define IF_MACRO 0
+#define TEST_IF_MACRO 0
 #define ELIF_MACRO 1
\ No newline at end of file
