Add missing file from previous commit.

diff --git a/Engine/Source/Developer/UnrealCodeAnalyzerTests/Private/ThreadSafety.cpp b/Engine/Source/Developer/UnrealCodeAnalyzerTests/Private/ThreadSafety.cpp
new file mode 100644
--- /dev/null
+++ b/Engine/Source/Developer/UnrealCodeAnalyzerTests/Private/ThreadSafety.cpp
@@ -0,0 +1,59 @@
+// Copyright 1998-2015 Epic Games, Inc. All Rights Reserved.
+
+#include "UnrealCodeAnalyzerTestsPCH.h"
+#include "ThreadSafety.h"
+
+int32 GGlobalValue;
+#define MY_DEFINE(x) GGlobalValue = x;
+
+void GGlobalFunc()
+{}
+
+class FClassWithStaticFunction
+{
+public:
+	static void StaticFunction()
+	{ }
+};
+
+class FStructWithStaticFunction
+{
+public:
+	static void StaticFunction()
+	{ }
+};
+
+void FunctionUsingInt(int32 Param)
+{
+
+}
+
+void UObjectDerivedThreadSafetyTest::PostInitProperties()
+{
+	MY_DEFINE(3);
+	GGlobalValue = 5;
+	GGlobalFunc();
+	FClassWithStaticFunction::StaticFunction();
+	FStructWithStaticFunction::StaticFunction();
+	FunctionUsingInt(GGlobalValue);
+}
+
+void UObjectDerivedThreadSafetyTest::Serialize(FArchive& Ar)
+{
+	MY_DEFINE(3);
+	GGlobalValue = 5;
+	GGlobalFunc();
+	FClassWithStaticFunction::StaticFunction();
+	FStructWithStaticFunction::StaticFunction();
+	FunctionUsingInt(GGlobalValue);
+}
+
+UObjectDerivedThreadSafetyTest::UObjectDerivedThreadSafetyTest()
+{
+	MY_DEFINE(4);
+	GGlobalFunc();
+	GGlobalValue = 5;
+	FClassWithStaticFunction::StaticFunction();
+	FStructWithStaticFunction::StaticFunction();
+	FunctionUsingInt(GGlobalValue);
+}
