Fixed QAGame crash on Android

diff --git a/Engine/Source/Developer/UnrealCodeAnalyzerTests/Private/ThreadSafety.cpp b/Engine/Source/Developer/UnrealCodeAnalyzerTests/Private/ThreadSafety.cpp
--- a/Engine/Source/Developer/UnrealCodeAnalyzerTests/Private/ThreadSafety.cpp
+++ b/Engine/Source/Developer/UnrealCodeAnalyzerTests/Private/ThreadSafety.cpp
@@ -30,6 +30,8 @@ void FunctionUsingInt(int32 Param)
 
 void UObjectDerivedThreadSafetyTest::PostInitProperties()
 {
+	Super::PostInitProperties();
+	
 	MY_DEFINE(3);
 	GGlobalValue = 5;
 	GGlobalFunc();
@@ -40,6 +42,8 @@ void UObjectDerivedThreadSafetyTest::PostInitProperties()
 
 void UObjectDerivedThreadSafetyTest::Serialize(FArchive& Ar)
 {
+	Super::Serialize(Ar);
+	
 	MY_DEFINE(3);
 	GGlobalValue = 5;
 	GGlobalFunc();
